HTML5LaunchHelper - fix typo

#rnx

diff --git a/Engine/Platforms/HTML5/Source/Programs/AutomationTool/HTML5Platform.Automation.cs b/Engine/Platforms/HTML5/Source/Programs/AutomationTool/HTML5Platform.Automation.cs
--- a/Engine/Platforms/HTML5/Source/Programs/AutomationTool/HTML5Platform.Automation.cs
+++ b/Engine/Platforms/HTML5/Source/Programs/AutomationTool/HTML5Platform.Automation.cs
@@ -275,7 +275,7 @@ public class HTML5Platform : Platform
 			File.Delete(ProjectBasename + ".cssgz");
 		}
 
-		if ( File.Exists(CombinePaths(CmdEnv.LocalRoot + "Engine/Binaries/DotNET/HTML5LaunchHelper.exe")) ) // TODO: platforms-extension: this currently needs to be built manually
+		if ( File.Exists(CombinePaths(CmdEnv.LocalRoot, "Engine/Binaries/DotNET/HTML5LaunchHelper.exe")) ) // TODO: platforms-extension: this currently needs to be built manually
 		{
 			File.Copy(CombinePaths(CmdEnv.LocalRoot, "Engine/Binaries/DotNET/HTML5LaunchHelper.exe"),CombinePaths(PackagePath, "HTML5LaunchHelper.exe"),true);
 		}
@@ -531,7 +531,7 @@ public class HTML5Platform : Platform
 		}
 
 		// Archive HTML5 Server and a Readme.
-		if ( File.Exists(CombinePaths(CmdEnv.LocalRoot + "Engine/Binaries/DotNET/HTML5LaunchHelper.exe")) ) // TODO: platforms-extension: this currently needs to be built manually
+		if ( File.Exists(CombinePaths(CmdEnv.LocalRoot, "Engine/Binaries/DotNET/HTML5LaunchHelper.exe")) ) // TODO: platforms-extension: this currently needs to be built manually
 		{
 			SC.ArchiveFiles(CombinePaths(CmdEnv.LocalRoot, "Engine/Binaries/DotNET/"), "HTML5LaunchHelper.exe");
 		}
