Fix HTML5 constexpr warnings
#rnx

diff --git a/Engine/Platforms/HTML5/Source/Runtime/Core/Public/HTML5Platform.h b/Engine/Platforms/HTML5/Source/Runtime/Core/Public/HTML5Platform.h
--- a/Engine/Platforms/HTML5/Source/Runtime/Core/Public/HTML5Platform.h
+++ b/Engine/Platforms/HTML5/Source/Runtime/Core/Public/HTML5Platform.h
@@ -32,6 +32,7 @@ typedef FHTML5Types FPlatformTypes;
 #define PLATFORM_MAX_FILEPATH_LENGTH_DEPRECATED				PATH_MAX
 #define PLATFORM_TCHAR_IS_4_BYTES							1
 #define PLATFORM_WCHAR_IS_4_BYTES							1
+#define PLATFORM_COMPILER_HAS_IF_CONSTEXPR					0
 #define PLATFORM_HAS_BSD_SOCKETS							1
 #define PLATFORM_HAS_BSD_SOCKET_FEATURE_MSG_DONTWAIT		1
 #define PLATFORM_HAS_NO_EPROCLIM							1
