Fix deprecation warning and remove removed ScopedPointer header.
#rnx

diff --git a/Engine/Platforms/HTML5/Source/Runtime/Launch/Private/LaunchHTML5.cpp b/Engine/Platforms/HTML5/Source/Runtime/Launch/Private/LaunchHTML5.cpp
--- a/Engine/Platforms/HTML5/Source/Runtime/Launch/Private/LaunchHTML5.cpp
+++ b/Engine/Platforms/HTML5/Source/Runtime/Launch/Private/LaunchHTML5.cpp
@@ -4,7 +4,6 @@
 #include "Misc/App.h"
 #include "Misc/OutputDeviceError.h"
 #include "HTML5PlatformMisc.h"
-#include "Templates/ScopedPointer.h"
 #include "LaunchEngineLoop.h"
 #include "EngineAnalytics.h"
 #include "Interfaces/IAnalyticsProvider.h"
diff --git a/Engine/Platforms/HTML5/Source/Runtime/NetworkFile/Private/HTTPTransport.cpp b/Engine/Platforms/HTML5/Source/Runtime/NetworkFile/Private/HTTPTransport.cpp
--- a/Engine/Platforms/HTML5/Source/Runtime/NetworkFile/Private/HTTPTransport.cpp
+++ b/Engine/Platforms/HTML5/Source/Runtime/NetworkFile/Private/HTTPTransport.cpp
@@ -35,7 +35,7 @@ bool FHTTPTransport::Initialize(const TCHAR* InHostIp)
 	// make sure that our string is again correctly formated
 	HostIp = FString::Printf(TEXT("http://%s"),*HostIp);
 
-	FCString::Strncpy(Url, *HostIp, ARRAY_COUNT(Url));
+	FCString::Strncpy(Url, *HostIp, UE_ARRAY_COUNT(Url));
 	emscripten_log(EM_LOG_CONSOLE , "Unreal File Server URL : %s ", TCHAR_TO_ANSI(Url));
 
 	TArray<uint8> In,Out;
