HTML5 as a platform extension cannot use GConfig anymore -- must use LoadLocalIniFile

#rnx

diff --git a/Engine/Platforms/HTML5/Source/Developer/HTML5TargetPlatform/Private/HTML5TargetPlatform.cpp b/Engine/Platforms/HTML5/Source/Developer/HTML5TargetPlatform/Private/HTML5TargetPlatform.cpp
--- a/Engine/Platforms/HTML5/Source/Developer/HTML5TargetPlatform/Private/HTML5TargetPlatform.cpp
+++ b/Engine/Platforms/HTML5/Source/Developer/HTML5TargetPlatform/Private/HTML5TargetPlatform.cpp
@@ -30,6 +30,9 @@ FCriticalSection FHTML5TargetPlatform::DevicesCriticalSection;
 
 FHTML5TargetPlatform::FHTML5TargetPlatform( )
 {
+	// must use HTML5EngineSettings (which comes from "LoadLocalIniFile").
+	// i.e. cannot use GConfig - for this contains only the active platform's ini files.
+	FConfigCacheIni::LoadLocalIniFile(HTML5EngineSettings, TEXT("Engine"), true, *PlatformName());
 	RefreshHTML5Setup();
 #if WITH_ENGINE
 	// load up texture settings from the config file
@@ -351,13 +354,15 @@ void FHTML5TargetPlatform::RefreshHTML5Setup()
 
 	// fill optional items first - these may be empty...
 	TArray<FString> DeviceMaps;
-	GConfig->GetArray( TEXT("/Script/HTML5PlatformEditor.HTML5SDKSettings"), TEXT("BrowserLauncher"), DeviceMaps, GEngineIni );
+//	GConfig->GetArray( TEXT("/Script/HTML5PlatformEditor.HTML5SDKSettings"), TEXT("BrowserLauncher"), DeviceMaps, GEngineIni );
+	HTML5EngineSettings.GetArray( TEXT("/Script/HTML5PlatformEditor.HTML5SDKSettings"), TEXT("BrowserLauncher"), DeviceMaps);
 	PopulateDevices(DeviceMaps, TEXT("user: "));
 	DefaultDeviceName.Empty(); // force default to one of the common browsers
 
 	// fill with "common browsers" (if they are installed)...
 	DeviceMaps.Empty();
-	GConfig->GetArray( TEXT("/Script/HTML5PlatformEditor.HTML5Browsers"), TEXT("BrowserLauncher"), DeviceMaps, GEngineIni );
+//	GConfig->GetArray( TEXT("/Script/HTML5PlatformEditor.HTML5Browsers"), TEXT("BrowserLauncher"), DeviceMaps, GEngineIni );
+	HTML5EngineSettings.GetArray( TEXT("/Script/HTML5PlatformEditor.HTML5Browsers"), TEXT("BrowserLauncher"), DeviceMaps);
 	PopulateDevices(DeviceMaps, TEXT(""));
 }
 
