undo one file checked in accidentally

#rnx

diff --git a/Engine/Platforms/HTML5/Source/Programs/UnrealBuildTool/HTML5SDKInfo.cs b/Engine/Platforms/HTML5/Source/Programs/UnrealBuildTool/HTML5SDKInfo.cs
--- a/Engine/Platforms/HTML5/Source/Programs/UnrealBuildTool/HTML5SDKInfo.cs
+++ b/Engine/Platforms/HTML5/Source/Programs/UnrealBuildTool/HTML5SDKInfo.cs
@@ -16,13 +16,13 @@ namespace UnrealBuildTool
 	public class HTML5SDKInfo
 	{
 		static string NODE_VER = "8.9.1_64bit";
-		static string SDKVersion = "1.38.31";
-
-//		static string NODE_VER = "12.9.1_64bit";
-//		static string SDKVersion = "1.39.0";
-
 		static string PYTHON_VER = "2.7.13.1_64bit"; // Only used on Windows; other platforms use built-in Python.
 
+		static string LLVM_VER = "e1.38.31_64bit";
+		static string SDKVersion = "1.38.31";
+
+//		static string LLVM_VER = "incoming";
+//		static string SDKVersion = "incoming";
 
 		// --------------------------------------------------
 		// --------------------------------------------------
@@ -32,15 +32,17 @@ namespace UnrealBuildTool
 			{
 				// If user has configured a custom Emscripten toolchain, use that automatically.
 				if (Environment.GetEnvironmentVariable("EMSDK") != null) return Environment.GetEnvironmentVariable("EMSDK");
+				// Otherwise, use the one embedded in this repository.
 
-				// platform extension path
+				// NEW PLATFORMS PATH
 				FileReference emsdkpath = FileReference.Combine(UnrealBuildTool.EnginePlatformExtensionsDirectory, "HTML5", "Build", "emsdk");
 				if (Directory.Exists(emsdkpath.FullName))
 				{
 					return emsdkpath.FullName;
 				}
-				Log.TraceErrorOnce("EMSCRIPTEN TOOLCHAIN NOT FOUND");
-				return "error_emscripten_toolchain_not_found";
+				// OLD ENGINE PATH
+//Log.TraceInformation("** WARNING : USING OLD EMSDK PATH ***");
+				return FileReference.Combine(UnrealBuildTool.EngineDirectory, "Extras", "ThirdPartyNotUE", "emsdk").FullName;
 			}
 		}
 		/// <summary>
@@ -54,19 +56,8 @@ namespace UnrealBuildTool
 				// If user has configured a custom Emscripten toolchain, use that automatically.
 				if (Environment.GetEnvironmentVariable("EMSCRIPTEN") != null) return Environment.GetEnvironmentVariable("EMSCRIPTEN");
 
-				// If user has configured a custom Emscripten toolchain, use that automatically.
-				string emscripten_root = GetEmscriptenConfigVar("EMSCRIPTEN_ROOT");
-				if (emscripten_root != null) return emscripten_root;
-
-				// platform extension path
-				string emscriptenpath = Path.Combine(SDKBase, "fastcomp", "emscripten");
-				if (Directory.Exists(emscriptenpath))
-				{
-					return emscriptenpath;
-				}
-
-				Log.TraceErrorOnce("EMSCRIPTEN NOT FOUND");
-				return "error_emscripten_not_found";
+				// Otherwise, use the one embedded in this repository.
+				return Path.Combine(SDKBase, "emscripten", SDKVersion);
 			}
 		}
 		/// <summary>
@@ -144,14 +135,8 @@ namespace UnrealBuildTool
 				string llvm_root = GetEmscriptenConfigVar("LLVM_ROOT");
 				if (llvm_root != null) return llvm_root;
 
-				// platform extension path
-				string llvmpath = Path.Combine(SDKBase, "fastcomp", "fastcomp", "bin");
-				if (Directory.Exists(llvmpath))
-				{
-					return llvmpath;
-				}
-				Log.TraceErrorOnce("LLVM NOT FOUND");
-				return "error_llvm_not_found";
+				// Otherwise, use the one embedded in this repository.
+				return Path.Combine(SDKBase, CURRENT_PLATFORM, "clang", LLVM_VER);
 			}
 		}
 		static string BINARYEN_ROOT
@@ -162,14 +147,9 @@ namespace UnrealBuildTool
 				string binaryen_root = GetEmscriptenConfigVar("BINARYEN_ROOT");
 				if (binaryen_root != null) return binaryen_root;
 
-				// platform extension path
-				string binaryenpath = Path.Combine(SDKBase, "fastcomp");
-				if (Directory.Exists(binaryenpath))
-				{
-					return binaryenpath;
-				}
-				Log.TraceErrorOnce("BINARYEN NOT FOUND");
-				return "error_binaryen_not_found";
+				// Otherwise, use the one embedded in this repository.
+//				return Path.Combine(SDKBase, CURRENT_PLATFORM, "binaryen", BINARYEN_VER);
+				return Path.Combine(SDKBase, CURRENT_PLATFORM, "clang", LLVM_VER, "binaryen");
 			}
 		}
 		static string NODE_JS
@@ -180,14 +160,8 @@ namespace UnrealBuildTool
 				string node_js = GetEmscriptenConfigVar("NODE_JS");
 				if (node_js != null) return node_js;
 
-				// platform extension path
-				string nodejspath = Path.Combine(SDKBase, "node", NODE_VER, "bin", "node" + PLATFORM_EXE);
-				if (File.Exists(nodejspath))
-				{
-					return nodejspath;
-				}
-				Log.TraceErrorOnce("NODEJS NOT FOUND");
-				return "error_nodejs_not_found";
+				// Otherwise, use the one embedded in this repository.
+				return Path.Combine(SDKBase, CURRENT_PLATFORM, "node", NODE_VER, "bin", "node" + PLATFORM_EXE);
 			}
 		}
 		static string PYTHON
@@ -428,26 +402,18 @@ namespace UnrealBuildTool
 // jic output dump is needed...
 //				processInfo.RedirectStandardError = true;
 //				processInfo.RedirectStandardOutput = true;
-				try
-				{
-					Process process = Process.Start(processInfo);
-//					process.OutputDataReceived += (object sender, DataReceivedEventArgs e) => Log.TraceInformation("output>>" + e.Data);
-//					process.BeginOutputReadLine();
-//					process.ErrorDataReceived += (object sender, DataReceivedEventArgs e) => Log.TraceInformation("error>>" + e.Data);
-//					process.BeginErrorReadLine();
-					process.WaitForExit();
-					Log.TraceInformation("emcc ExitCode: {0}", process.ExitCode);
-					process.Close();
-				}
-				catch (System.ComponentModel.Win32Exception ex)
-				{
-					// Process.Start() as terminated quick enough betore control has returned to process.WaitForExit()
-					Log.TraceInformation("Win32Exception ex.NativeErrorCode: {0}", ex.NativeErrorCode);
-				}
+				Process process = Process.Start(processInfo);
+//				process.OutputDataReceived += (object sender, DataReceivedEventArgs e) => Log.TraceInformation("output>>" + e.Data);
+//				process.BeginOutputReadLine();
+//				process.ErrorDataReceived += (object sender, DataReceivedEventArgs e) => Log.TraceInformation("error>>" + e.Data);
+//				process.BeginErrorReadLine();
+				process.WaitForExit();
+				Log.TraceInformation("emcc ExitCode: {0}", process.ExitCode);
+				process.Close();
 				// uncomment OPTIMIZER (GUBP on build machines needs this)
 				// and PYTHON (reduce warnings on EMCC_DEBUG=1)
 				string pyth = Regex.Replace(PYTHON, @"\\", @"\\");
-				string optz = Regex.Replace(Path.Combine(BINARYEN_ROOT, "bin", "optimizer") + PLATFORM_EXE, @"\\", @"\\");
+				string optz = Regex.Replace(Path.Combine(LLVM_ROOT, "optimizer") + PLATFORM_EXE, @"\\", @"\\");
 				string byn = Regex.Replace(BINARYEN_ROOT, @"\\", @"\\");
 				string txt = Regex.Replace(
 					Regex.Replace(
@@ -493,28 +459,6 @@ namespace UnrealBuildTool
 			Environment.SetEnvironmentVariable("NODE", NODE_JS);
 			Environment.SetEnvironmentVariable("LLVM", LLVM_ROOT);
 
-			// --------------------------------------------------
-			// the following is needed when UE4 from GitHub on Linux is used
-			if (BuildHostPlatform.Current.Platform == UnrealTargetPlatform.Linux)
-			{
-				string findpath = "\"" + Path.Combine(EMSCRIPTEN_ROOT, "system", "lib") + "\"";
-				ProcessStartInfo findproc = new ProcessStartInfo("find", findpath + " -type f -name \"*symbols\" -exec dos2unix {} \\;");
-				findproc.CreateNoWindow = true;
-				findproc.UseShellExecute = false;
-				try
-				{
-					Process process = Process.Start(findproc);
-					process.WaitForExit();
-					Log.TraceInformation("find symbols dos2unix conversions ExitCode: {0}", process.ExitCode);
-					process.Close();
-				}
-				catch (System.ComponentModel.Win32Exception ex)
-				{
-					// Process.Start() as terminated quick enough betore control has returned to process.WaitForExit()
-					Log.TraceInformation("Win32Exception ex.NativeErrorCode: {0}", ex.NativeErrorCode);
-				}
-			}
-
 			return DOT_EMSCRIPTEN;
 		}
 
@@ -554,16 +498,6 @@ namespace UnrealBuildTool
 			return new FileReference(PYTHON);
 		}
 
-		/// <summary>
-		///
-		/// </summary>
-		/// <returns></returns>
-		public static string MacPythonLib() // UE-75402
-		{
-			string UE4PythonPath = FileReference.Combine(UnrealBuildTool.EngineDirectory, "Binaries", "ThirdParty", "Python").FullName;
-			return Path.Combine(UE4PythonPath, "Mac", "lib", "python2.7", "lib-dynload");
-		}
-
 		/// <summary>
 		///
 		/// </summary>
@@ -611,3 +545,4 @@ namespace UnrealBuildTool
 		// BuildPhysX.Automation.cs
 	}
 }
+
