HTML5 remove Deprecation Warning banner

[RN]

diff --git a/Engine/Platforms/HTML5/Source/Developer/HTML5PlatformEditor/Private/HTML5TargetSettingsCustomization.cpp b/Engine/Platforms/HTML5/Source/Developer/HTML5PlatformEditor/Private/HTML5TargetSettingsCustomization.cpp
--- a/Engine/Platforms/HTML5/Source/Developer/HTML5PlatformEditor/Private/HTML5TargetSettingsCustomization.cpp
+++ b/Engine/Platforms/HTML5/Source/Developer/HTML5PlatformEditor/Private/HTML5TargetSettingsCustomization.cpp
@@ -102,9 +102,9 @@ FHTML5TargetSettingsCustomization::~FHTML5TargetSettingsCustomization()
 
 void FHTML5TargetSettingsCustomization::CustomizeDetails(IDetailLayoutBuilder& DetailLayout)
 {
-	BuildDeprecationMessage(DetailLayout); // !!! HTML5 Deprecation Message
+//	BuildDeprecationMessage(DetailLayout); // !!! HTML5 Deprecation Message
 	// --------------------------------------------------
-	BuildEmscriptenCategory(DetailLayout); // !!! HTML5 Deprecation Message
+	BuildEmscriptenCategory(DetailLayout);
 	AudioPluginWidgetManager.BuildAudioCategory(DetailLayout, FString("HTML5"));
 }
 
