HTML5 - more code removal - source code section

#rnx

diff --git a/Engine/Platforms/HTML5/Source/Developer/HTML5PlatformEditor/Classes/HTML5TargetSettings.h b/Engine/Platforms/HTML5/Source/Developer/HTML5PlatformEditor/Classes/HTML5TargetSettings.h
--- a/Engine/Platforms/HTML5/Source/Developer/HTML5PlatformEditor/Classes/HTML5TargetSettings.h
+++ b/Engine/Platforms/HTML5/Source/Developer/HTML5PlatformEditor/Classes/HTML5TargetSettings.h
@@ -38,6 +38,13 @@ public:
 
 	GENERATED_UCLASS_BODY()
 
+	// ------------------------------------------------------------
+	// Use HTML5 .ini files in the editor
+	virtual const TCHAR* GetConfigOverridePlatform() const override
+	{
+		return TEXT("HTML5");
+	}
+
 	// ------------------------------------------------------------
 
 	/**
diff --git a/Engine/Platforms/HTML5/Source/Runtime/Sockets/Sockets_HTML5.Build.cs b/Engine/Platforms/HTML5/Source/Runtime/Sockets/Sockets_HTML5.Build.cs
--- a/Engine/Platforms/HTML5/Source/Runtime/Sockets/Sockets_HTML5.Build.cs
+++ b/Engine/Platforms/HTML5/Source/Runtime/Sockets/Sockets_HTML5.Build.cs
@@ -6,5 +6,6 @@ public class Sockets_HTML5 : Sockets
 {
 	public Sockets_HTML5(ReadOnlyTargetRules Target) : base(Target)
 	{
+		PublicDefinitions.Add("PLATFORM_SOCKETSUBSYSTEM=FName(TEXT(\"HTML5\"))");
 	}
 }
diff --git a/Engine/Source/Editor/MainFrame/Private/Frame/MainFrameActions.cpp b/Engine/Source/Editor/MainFrame/Private/Frame/MainFrameActions.cpp
--- a/Engine/Source/Editor/MainFrame/Private/Frame/MainFrameActions.cpp
+++ b/Engine/Source/Editor/MainFrame/Private/Frame/MainFrameActions.cpp
@@ -734,10 +734,7 @@ void FMainFrameActionCallbacks::PackageProject( const FName InPlatformInfoName )
 
 	if (PackagingSettings->UsePakFile)
 	{
-	  if (PlatformInfo->TargetPlatformName != FName("HTML5")) 
-	  { 
 		OptionalParams += TEXT(" -pak");
-	  }
 	}
 
 	if (PackagingSettings->IncludePrerequisites)
diff --git a/Engine/Source/Runtime/Linux/AudioMixerSDL/Private/HTML5/AudioMixerPlatformDefinesSDL.h b/Engine/Source/Runtime/Linux/AudioMixerSDL/Private/HTML5/AudioMixerPlatformDefinesSDL.h
deleted file mode 100644
diff --git a/Engine/Source/Runtime/Sockets/Public/SocketSubsystem.h b/Engine/Source/Runtime/Sockets/Public/SocketSubsystem.h
--- a/Engine/Source/Runtime/Sockets/Public/SocketSubsystem.h
+++ b/Engine/Source/Runtime/Sockets/Public/SocketSubsystem.h
@@ -30,8 +30,6 @@ SOCKETS_API DECLARE_LOG_CATEGORY_EXTERN(LogSockets, Log, All);
 		#define PLATFORM_SOCKETSUBSYSTEM FName(TEXT("PS4"))
 	#elif PLATFORM_XBOXONE
 		#define PLATFORM_SOCKETSUBSYSTEM FName(TEXT("XBOXONE"))
-	#elif defined(__EMSCRIPTEN__)
-		#define PLATFORM_SOCKETSUBSYSTEM FName(TEXT("HTML5"))
 	#elif PLATFORM_SWITCH
 		#define PLATFORM_SOCKETSUBSYSTEM FName(TEXT("SWITCH"))
 	#else
diff --git a/Engine/Source/Runtime/UnrealAudio/Public/UnrealAudioTypes.h b/Engine/Source/Runtime/UnrealAudio/Public/UnrealAudioTypes.h
--- a/Engine/Source/Runtime/UnrealAudio/Public/UnrealAudioTypes.h
+++ b/Engine/Source/Runtime/UnrealAudio/Public/UnrealAudioTypes.h
@@ -83,7 +83,6 @@ namespace UAudio
 			ALSA,
 			CORE_AUDIO,
 			OPENAL,
-			HTML5,
 			DUMMY,
 		};
 
@@ -97,7 +96,6 @@ namespace UAudio
 				case ALSA:			return TEXT("ALSA");
 				case CORE_AUDIO:	return TEXT("CORE_AUDIO");
 				case OPENAL:		return TEXT("OPENAL");
-				case HTML5:			return TEXT("HTML5");
 				case DUMMY:			return TEXT("DUMMY");
 				default:			check(false);
 			}
