HTML5 - more code removal - opengl driver and shader code only

#rnx

diff --git a/Engine/Platforms/HTML5/Source/Runtime/OpenGLDrv/OpenGLDrv_HTML5.Build.cs b/Engine/Platforms/HTML5/Source/Runtime/OpenGLDrv/OpenGLDrv_HTML5.Build.cs
--- a/Engine/Platforms/HTML5/Source/Runtime/OpenGLDrv/OpenGLDrv_HTML5.Build.cs
+++ b/Engine/Platforms/HTML5/Source/Runtime/OpenGLDrv/OpenGLDrv_HTML5.Build.cs
@@ -2,7 +2,7 @@
 
 using UnrealBuildTool;
 
-[SupportedPlatformsAttribute(new string[] {"Win64", "Mac", "Linux", "HTML5", "LinuxAArch64"})]
+[SupportedPlatformsAttribute(new string[] {"HTML5"})]
 public class OpenGLDrv_HTML5 : OpenGLDrv
 {
 	public OpenGLDrv_HTML5(ReadOnlyTargetRules Target) : base(Target)
diff --git a/Engine/Platforms/HTML5/Source/Runtime/OpenGLDrv/Private/HTML5OpenGL.cpp b/Engine/Platforms/HTML5/Source/Runtime/OpenGLDrv/Private/HTML5OpenGL.cpp
--- a/Engine/Platforms/HTML5/Source/Runtime/OpenGLDrv/Private/HTML5OpenGL.cpp
+++ b/Engine/Platforms/HTML5/Source/Runtime/OpenGLDrv/Private/HTML5OpenGL.cpp
@@ -10,6 +10,8 @@ THIRD_PARTY_INCLUDES_START
 THIRD_PARTY_INCLUDES_END
 #include "HTML5JavaScriptFx.h"
 
+#include "OpenGLShaders.h"
+
 DEFINE_LOG_CATEGORY_STATIC(LogHTML5OpenGL, Log, All);
 
 bool FHTML5OpenGL::bCombinedDepthStencilAttachment = false;
@@ -150,6 +152,86 @@ void FHTML5OpenGL::ProcessExtensions( const FString& ExtensionsString )
 	UE_LOG(LogRHI, Log, TEXT("Fragment shader highp precision: %d"), ShaderHighPrecision);
 }
 
+void FHTML5OpenGL::PE_GetCurrentOpenGLShaderDeviceCapabilities(FOpenGLShaderDeviceCapabilities& Capabilities)
+{
+//	Capabilities.TargetPlatform = EOpenGLShaderTargetPlatform::OGLSTP_HTML5;
+	Capabilities.TargetPlatform = EOpenGLShaderTargetPlatform::OGLSTP_Unknown; // jic
+	Capabilities.bUseES30ShadingLanguage = FHTML5OpenGL::UseES30ShadingLanguage();
+	Capabilities.bSupportsShaderTextureLod = FHTML5OpenGL::SupportsShaderTextureLod();
+}
+
+bool FHTML5OpenGL::PE_GLSLToDeviceCompatibleGLSL(FAnsiCharArray& GlslCodeOriginal, const FString& ShaderName, GLenum TypeEnum, const FOpenGLShaderDeviceCapabilities& Capabilities, FAnsiCharArray& GlslCode)
+{
+	// NOTE: this was an exact copy of Engine/Source/Runtime/OpenGLDrv/Private/OpenGLShaders.cpp::GLSLToDeviceCompatibleGLSL()
+	// with all obvious TargetPlatform & PLATFORM_XYZ code removed -- and then stripped out unsupported shader level code afterwards
+	// the original indentation has also been kept here for reference
+	// *** make sure to view 4.23 or earlier to see where OGLSTP_HTML5 calls were before for reference
+
+
+//	// Whether shader was compiled for ES 3.1
+	const ANSICHAR* ES310Version = "#version 310 es";
+	const bool bES31 = (FCStringAnsi::Strstr(GlslCodeOriginal.GetData(), ES310Version) != nullptr);
+	bool bUseES30ShadingLanguage = Capabilities.bUseES30ShadingLanguage;
+
+	bool bNeedsExtDrawInstancedDefine = !bES31;
+//	if (Capabilities.TargetPlatform == EOpenGLShaderTargetPlatform::OGLSTP_Android || Capabilities.TargetPlatform == EOpenGLShaderTargetPlatform::OGLSTP_HTML5)
+//		else if (IsES2Platform(Capabilities.MaxRHIShaderPlatform))
+			// #version NNN has to be the first line in the file, so it has to be added before anything else.
+			if (bUseES30ShadingLanguage)
+			{
+				bNeedsExtDrawInstancedDefine = false;
+				PE_AppendCString(GlslCode, "#version 300 es\n");
+			}
+			else
+			{
+				PE_AppendCString(GlslCode, "#version 100\n");
+			}
+			PE_ReplaceCString(GlslCodeOriginal, "#version 100", "");
+
+	if (bNeedsExtDrawInstancedDefine)
+	{
+		// Check for the GL_EXT_draw_instanced extension if necessary (version < 300)
+		PE_AppendCString(GlslCode, "#ifdef GL_EXT_draw_instanced\n");
+		PE_AppendCString(GlslCode, "#define UE_EXT_draw_instanced 1\n");
+		PE_AppendCString(GlslCode, "#endif\n");
+	}
+
+	if (ShaderName.IsEmpty() == false)
+	{
+		PE_AppendCString(GlslCode, "// ");
+		PE_AppendCString(GlslCode, TCHAR_TO_ANSI(ShaderName.GetCharArray().GetData()));
+		PE_AppendCString(GlslCode, "\n");
+	}
+
+//	if (Capabilities.TargetPlatform == EOpenGLShaderTargetPlatform::OGLSTP_Android)
+//	elseif (Capabilities.TargetPlatform == EOpenGLShaderTargetPlatform::OGLSTP_HTML5)
+		// HTML5 use case is much simpler, use a separate chunk of code from android.
+		if (!Capabilities.bSupportsShaderTextureLod)
+		{
+			PE_AppendCString(GlslCode,
+				"#define DONTEMITEXTENSIONSHADERTEXTURELODENABLE \n"
+				"#define texture2DLodEXT(a, b, c) texture2D(a, b) \n"
+				"#define textureCubeLodEXT(a, b, c) textureCube(a, b) \n");
+		}
+
+	if (FOpenGL::SupportsClipControl())
+	{
+		PE_AppendCString(GlslCode, "#define HLSLCC_DX11ClipSpace 0 \n");
+	}
+	else
+	{
+		PE_AppendCString(GlslCode, "#define HLSLCC_DX11ClipSpace 1 \n");
+	}
+
+	// Append the possibly edited shader to the one we will compile.
+	// This is to make it easier to debug as we can see the whole
+	// shader source.
+	PE_AppendCString(GlslCode, "\n\n");
+	PE_AppendCString(GlslCode, GlslCodeOriginal.GetData());
+
+	return true;
+}
+
 struct FPlatformOpenGLContext
 {
 	GLuint			ViewportFramebuffer;
diff --git a/Engine/Platforms/HTML5/Source/Runtime/OpenGLDrv/Private/HTML5OpenGL.h b/Engine/Platforms/HTML5/Source/Runtime/OpenGLDrv/Private/HTML5OpenGL.h
--- a/Engine/Platforms/HTML5/Source/Runtime/OpenGLDrv/Private/HTML5OpenGL.h
+++ b/Engine/Platforms/HTML5/Source/Runtime/OpenGLDrv/Private/HTML5OpenGL.h
@@ -259,6 +259,10 @@ struct FHTML5OpenGL : public FOpenGLES2
 		return SP_OPENGL_ES2_WEBGL;
 	}
 
+	// platform extensions overrides
+	static void PE_GetCurrentOpenGLShaderDeviceCapabilities(FOpenGLShaderDeviceCapabilities& Capabilities);
+	static bool PE_GLSLToDeviceCompatibleGLSL(FAnsiCharArray& GlslCodeOriginal, const FString& ShaderName, GLenum TypeEnum, const FOpenGLShaderDeviceCapabilities& Capabilities, FAnsiCharArray& GlslCode);
+
 	static FORCEINLINE FString GetAdapterName() { return TEXT(""); }
 
 	// TODO: Make this true or false depending on whether WebGL 2 is available or not.
diff --git a/Engine/Platforms/HTML5/Source/Runtime/OpenGLDrv/Private/HTML5OpenGLDrvPrivate.h b/Engine/Platforms/HTML5/Source/Runtime/OpenGLDrv/Private/HTML5OpenGLDrvPrivate.h
new file mode 100644
--- /dev/null
+++ b/Engine/Platforms/HTML5/Source/Runtime/OpenGLDrv/Private/HTML5OpenGLDrvPrivate.h
@@ -0,0 +1,8 @@
+// Copyright 1998-2019 Epic Games, Inc. All Rights Reserved.
+
+/*=============================================================================
+	*HTML5*OpenGLDrvPrivate.h: Platform Extension: Private HTML5 OpenGL includes
+=============================================================================*/
+
+#include "HTML5OpenGL.h"
+
diff --git a/Engine/Source/Runtime/OpenGLDrv/Private/OpenGLDrv.cpp b/Engine/Source/Runtime/OpenGLDrv/Private/OpenGLDrv.cpp
--- a/Engine/Source/Runtime/OpenGLDrv/Private/OpenGLDrv.cpp
+++ b/Engine/Source/Runtime/OpenGLDrv/Private/OpenGLDrv.cpp
@@ -19,6 +19,7 @@ IMPLEMENT_MODULE(FOpenGLDynamicRHIModule, OpenGLDrv);
 
 #include "Shader.h"
 #include "OneColorShader.h"
+#include "OpenGLShaders.h"
 
 /** OpenGL Logging. */
 DEFINE_LOG_CATEGORY(LogOpenGL);
@@ -573,6 +574,11 @@ void FOpenGLBase::ProcessExtensions( const FString& ExtensionsString )
 #endif
 }
 
+void PE_GetCurrentOpenGLShaderDeviceCapabilities(FOpenGLShaderDeviceCapabilities& Capabilities)
+{
+	Capabilities.TargetPlatform = EOpenGLShaderTargetPlatform::OGLSTP_Unknown;
+}
+
 void GetExtensionsString( FString& ExtensionsString)
 {
 	GLint ExtensionCount = 0;
diff --git a/Engine/Source/Runtime/OpenGLDrv/Private/OpenGLShaders.cpp b/Engine/Source/Runtime/OpenGLDrv/Private/OpenGLShaders.cpp
--- a/Engine/Source/Runtime/OpenGLDrv/Private/OpenGLShaders.cpp
+++ b/Engine/Source/Runtime/OpenGLDrv/Private/OpenGLShaders.cpp
@@ -534,6 +534,17 @@ namespace
 	}
 }
 
+// make some anon ns functions available to platform extensions
+void PE_AppendCString(TArray<ANSICHAR> & Dest, const ANSICHAR * Source)
+{
+	AppendCString(Dest, Source);
+}
+
+void PE_ReplaceCString(TArray<ANSICHAR> & Dest, const ANSICHAR * Source, const ANSICHAR * Replacement)
+{
+	ReplaceCString(Dest, Source, Replacement);
+}
+
 
 FOpenGLProgramBinaryCache::FPreviousGLProgramBinaryCacheInfo::FPreviousGLProgramBinaryCacheInfo() : NumberOfOldEntriesReused(0) {}
 FOpenGLProgramBinaryCache::FPreviousGLProgramBinaryCacheInfo::FPreviousGLProgramBinaryCacheInfo(FOpenGLProgramBinaryCache::FPreviousGLProgramBinaryCacheInfo&&) = default;
@@ -947,14 +958,10 @@ void OPENGLDRV_API GetCurrentOpenGLShaderDeviceCapabilities(FOpenGLShaderDeviceC
 		Capabilities.bRequiresTexture2DPrecisionHack = FOpenGL::RequiresTexture2DPrecisionHack();
 		Capabilities.bRequiresRoundFunctionHack = FOpenGL::RequiresRoundFunctionHack();
 	#endif // PLATFORM_LUMINGL4
-#elif defined(__EMSCRIPTEN__)
-	Capabilities.TargetPlatform = EOpenGLShaderTargetPlatform::OGLSTP_HTML5;
-	Capabilities.bUseES30ShadingLanguage = FOpenGL::UseES30ShadingLanguage();
-	Capabilities.bSupportsShaderTextureLod = FOpenGL::SupportsShaderTextureLod();
 #elif PLATFORM_IOS
 	Capabilities.TargetPlatform = EOpenGLShaderTargetPlatform::OGLSTP_iOS;
 #else
-	Capabilities.TargetPlatform = EOpenGLShaderTargetPlatform::OGLSTP_Unknown;
+	FOpenGL::PE_GetCurrentOpenGLShaderDeviceCapabilities(Capabilities); // platform extension
 #endif
 	Capabilities.MaxRHIShaderPlatform = GMaxRHIShaderPlatform;
 	Capabilities.bSupportsSeparateShaderObjects = FOpenGL::SupportsSeparateShaderObjects();
@@ -967,6 +974,11 @@ void OPENGLDRV_API GetCurrentOpenGLShaderDeviceCapabilities(FOpenGLShaderDeviceC
 
 void OPENGLDRV_API GLSLToDeviceCompatibleGLSL(FAnsiCharArray& GlslCodeOriginal, const FString& ShaderName, GLenum TypeEnum, const FOpenGLShaderDeviceCapabilities& Capabilities, FAnsiCharArray& GlslCode)
 {
+	if (FOpenGL::PE_GLSLToDeviceCompatibleGLSL(GlslCodeOriginal, ShaderName, TypeEnum, Capabilities, GlslCode))
+	{
+		return; // platform extension overrides
+	}
+
 	// Whether shader was compiled for ES 3.1
 	const ANSICHAR* ES310Version = "#version 310 es";
 	const bool bES31 = (FCStringAnsi::Strstr(GlslCodeOriginal.GetData(), ES310Version) != nullptr);
@@ -993,7 +1005,7 @@ void OPENGLDRV_API GLSLToDeviceCompatibleGLSL(FAnsiCharArray& GlslCodeOriginal,
 	bool bGlslCodeHasExtensions = CStringCountOccurances(GlslCodeOriginal, GlslPlaceHolderAfterExtensions) == 1;
 	
 	bool bNeedsExtDrawInstancedDefine = false;
-	if (Capabilities.TargetPlatform == EOpenGLShaderTargetPlatform::OGLSTP_Android || Capabilities.TargetPlatform == EOpenGLShaderTargetPlatform::OGLSTP_HTML5)
+	if (Capabilities.TargetPlatform == EOpenGLShaderTargetPlatform::OGLSTP_Android)
 	{
 		bNeedsExtDrawInstancedDefine = !bES31;
 		PRAGMA_DISABLE_DEPRECATION_WARNINGS
@@ -1370,17 +1382,6 @@ void OPENGLDRV_API GLSLToDeviceCompatibleGLSL(FAnsiCharArray& GlslCodeOriginal,
 		}
 		PRAGMA_ENABLE_DEPRECATION_WARNINGS
 	}
-	else if (Capabilities.TargetPlatform == EOpenGLShaderTargetPlatform::OGLSTP_HTML5)
-	{
-		// HTML5 use case is much simpler, use a separate chunk of code from android.
-		if (!Capabilities.bSupportsShaderTextureLod)
-		{
-			AppendCString(GlslCode,
-				"#define DONTEMITEXTENSIONSHADERTEXTURELODENABLE \n"
-				"#define texture2DLodEXT(a, b, c) texture2D(a, b) \n"
-				"#define textureCubeLodEXT(a, b, c) textureCube(a, b) \n");
-		}
-	}
 
 	if (FOpenGL::SupportsClipControl())
 	{
diff --git a/Engine/Source/Runtime/OpenGLDrv/Public/OpenGL.h b/Engine/Source/Runtime/OpenGLDrv/Public/OpenGL.h
--- a/Engine/Source/Runtime/OpenGLDrv/Public/OpenGL.h
+++ b/Engine/Source/Runtime/OpenGLDrv/Public/OpenGL.h
@@ -47,6 +47,14 @@ struct FPlatformOpenGLContext;
 #define GL_RGBA16F    0x881A
 #endif
 
+
+// for platform extensions: from OpenGLShaders.h
+extern "C" struct FOpenGLShaderDeviceCapabilities;
+
+// for platform extensions: from OpenGLResources.h
+typedef TArray<ANSICHAR> FAnsiCharArray;
+
+
 // Base static class
 class FOpenGLBase
 {
@@ -397,6 +405,10 @@ public:
 		return PreferredPixelFormat;
 	}
 
+	// for platform extensions
+	static void PE_GetCurrentOpenGLShaderDeviceCapabilities(FOpenGLShaderDeviceCapabilities& Capabilities);
+	static bool PE_GLSLToDeviceCompatibleGLSL(FAnsiCharArray& GlslCodeOriginal, const FString& ShaderName, GLenum TypeEnum, const FOpenGLShaderDeviceCapabilities& Capabilities, FAnsiCharArray& GlslCode) UGL_OPTIONAL(false)
+
 protected:
 	static GLint MaxTextureImageUnits;
 	static GLint MaxCombinedTextureImageUnits;
diff --git a/Engine/Source/Runtime/OpenGLDrv/Public/OpenGLDrv.h b/Engine/Source/Runtime/OpenGLDrv/Public/OpenGLDrv.h
--- a/Engine/Source/Runtime/OpenGLDrv/Public/OpenGLDrv.h
+++ b/Engine/Source/Runtime/OpenGLDrv/Public/OpenGLDrv.h
@@ -19,22 +19,20 @@
 // @todo platplug: Replace all of these includes with a call to COMPILED_PLATFORM_HEADER(OpenGLDrvPrivate.h)
 //TODO: Move these to OpenGLDrvPrivate.h
 #if PLATFORM_WINDOWS
-#include "Runtime/OpenGLDrv/Private/Windows/OpenGLWindows.h"
+	#include "Runtime/OpenGLDrv/Private/Windows/OpenGLWindows.h"
 #elif PLATFORM_LINUX
-#include "Runtime/OpenGLDrv/Private/Linux/OpenGLLinux.h"
+	#include "Runtime/OpenGLDrv/Private/Linux/OpenGLLinux.h"
 #elif PLATFORM_LUMIN
-// these guys will self-select
-#include "Lumin/LuminESDeferredOpenGL.h"
-#include "Lumin/LuminOpenGL.h"
-#include "Lumin/LuminGL4.h"
+	// these guys will self-select
+	#include "Lumin/LuminESDeferredOpenGL.h"
+	#include "Lumin/LuminOpenGL.h"
+	#include "Lumin/LuminGL4.h"
 #elif PLATFORM_ANDROIDESDEFERRED
-#include "Android/AndroidESDeferredOpenGL.h"
+	#include "Android/AndroidESDeferredOpenGL.h"
 #elif PLATFORM_ANDROID
-#include "Android/AndroidOpenGL.h"
-#elif defined(__EMSCRIPTEN__)
-#include "HTML5OpenGL.h"
-#elif PLATFORM_LINUX
-#include "Linux/OpenGLLinux.h"
+	#include "Android/AndroidOpenGL.h"
+#else
+#include COMPILED_PLATFORM_HEADER(OpenGLDrvPrivate.h)
 #endif
 
 // Define here so don't have to do platform filtering
diff --git a/Engine/Source/Runtime/OpenGLDrv/Public/OpenGLShaders.h b/Engine/Source/Runtime/OpenGLDrv/Public/OpenGLShaders.h
--- a/Engine/Source/Runtime/OpenGLDrv/Public/OpenGLShaders.h
+++ b/Engine/Source/Runtime/OpenGLDrv/Public/OpenGLShaders.h
@@ -14,7 +14,6 @@ enum EOpenGLShaderTargetPlatform
 	OGLSTP_Unknown,
 	OGLSTP_Desktop,
 	OGLSTP_Android,
-	OGLSTP_HTML5,
 	OGLSTP_iOS,
 };
 
@@ -62,3 +61,9 @@ void OPENGLDRV_API GetCurrentOpenGLShaderDeviceCapabilities(FOpenGLShaderDeviceC
   * @param GlslCode - [out] Compilable GLSL
   */
 void OPENGLDRV_API GLSLToDeviceCompatibleGLSL(FAnsiCharArray& GlslCodeOriginal, const FString& ShaderName, GLenum TypeEnum, const FOpenGLShaderDeviceCapabilities& Capabilities, FAnsiCharArray& GlslCode);
+
+
+// make some anon ns functions available to platform extensions
+extern "C" void PE_AppendCString(TArray<ANSICHAR> & Dest, const ANSICHAR * Source);
+extern "C" void PE_ReplaceCString(TArray<ANSICHAR> & Dest, const ANSICHAR * Source, const ANSICHAR * Replacement);
+
diff --git a/Engine/Source/Runtime/RHI/Private/RHI.cpp b/Engine/Source/Runtime/RHI/Private/RHI.cpp
--- a/Engine/Source/Runtime/RHI/Private/RHI.cpp
+++ b/Engine/Source/Runtime/RHI/Private/RHI.cpp
@@ -693,7 +693,6 @@ static FName NAME_PLATFORM_IOS(TEXT("IOS"));
 static FName NAME_PLATFORM_MAC(TEXT("Mac"));
 static FName NAME_PLATFORM_SWITCH(TEXT("Switch"));
 static FName NAME_PLATFORM_TVOS(TEXT("TVOS"));
-static FName NAME_PLATFORM_HTML5(TEXT("HTML5"));
 static FName NAME_PLATFORM_LUMIN(TEXT("Lumin"));
 
 // @todo platplug: This is still here, only being used now by UMaterialShaderQualitySettings::GetOrCreatePlatformSettings
@@ -723,7 +722,7 @@ FName ShaderPlatformToPlatformName(EShaderPlatform Platform)
 	case SP_OPENGL_ES3_1_ANDROID:
 		return NAME_PLATFORM_ANDROID;
 	case SP_OPENGL_ES2_WEBGL:
-		return NAME_PLATFORM_HTML5;
+		return FName(TEXT(PREPROCESSOR_TO_STRING(PLATFORM_HEADER_NAME))); // WIP: platform extension magic
 	case SP_METAL:
 	case SP_METAL_MRT:
 		return NAME_PLATFORM_IOS;
