CIS build error fix

diff --git a/Engine/Platforms/HTML5/Source/Runtime/Online/SSL/SSL_HTML5.Build.cs b/Engine/Platforms/HTML5/Source/Runtime/Online/SSL/SSL_HTML5.Build.cs
--- a/Engine/Platforms/HTML5/Source/Runtime/Online/SSL/SSL_HTML5.Build.cs
+++ b/Engine/Platforms/HTML5/Source/Runtime/Online/SSL/SSL_HTML5.Build.cs
@@ -4,7 +4,7 @@ using UnrealBuildTool;
 
 public class SSL_HTML5 : SSL
 {
-	protected virtual bool PlatformSupportsSSL { get { return true; } }
+	protected override bool PlatformSupportsSSL { get { return true; } }
 
     public SSL_HTML5(ReadOnlyTargetRules Target) : base(Target)
     {
