Slate: add HTML5 platform text field.

diff --git a/Engine/Platforms/HTML5/Source/Runtime/Slate/Public/Framework/Text/HTML5PlatformTextField.h b/Engine/Platforms/HTML5/Source/Runtime/Slate/Public/Framework/Text/HTML5PlatformTextField.h
new file mode 100644
--- /dev/null
+++ b/Engine/Platforms/HTML5/Source/Runtime/Slate/Public/Framework/Text/HTML5PlatformTextField.h
@@ -0,0 +1,5 @@
+// Copyright 1998-2019 Epic Games, Inc. All Rights Reserved.
+
+#pragma once
+
+#include "Framework/Text/GenericPlatformTextField.h"
diff --git a/Engine/Platforms/HTML5/Source/Runtime/Slate/Slate_HTML5.Build.cs b/Engine/Platforms/HTML5/Source/Runtime/Slate/Slate_HTML5.Build.cs
new file mode 100644
--- /dev/null
+++ b/Engine/Platforms/HTML5/Source/Runtime/Slate/Slate_HTML5.Build.cs
@@ -0,0 +1,12 @@
+// Copyright 1998-2019 Epic Games, Inc. All Rights Reserved.
+
+using UnrealBuildTool;
+using System;
+using System.IO;
+
+public class Slate_HTML5 : Slate
+{
+	public Slate_HTML5(ReadOnlyTargetRules Target) : base(Target)
+	{
+	}
+}
