add missing HTML5AudioMixerSDLDefines.h file

#rnx

diff --git a/Engine/Platforms/HTML5/Source/Runtime/AudioMixer/AudioMixer_HTML5.Build.cs b/Engine/Platforms/HTML5/Source/Runtime/AudioMixer/AudioMixer_HTML5.Build.cs
--- a/Engine/Platforms/HTML5/Source/Runtime/AudioMixer/AudioMixer_HTML5.Build.cs
+++ b/Engine/Platforms/HTML5/Source/Runtime/AudioMixer/AudioMixer_HTML5.Build.cs
@@ -6,7 +6,6 @@ namespace UnrealBuildTool.Rules
 	{
 		public AudioMixer_HTML5(ReadOnlyTargetRules Target) : base(Target)
 		{
-			// pretty syre we can delete this!
 		}
 	}
 }
diff --git a/Engine/Platforms/HTML5/Source/Runtime/AudioMixerSDL/AudioMixerSDL_HTML5.Build.cs b/Engine/Platforms/HTML5/Source/Runtime/AudioMixerSDL/AudioMixerSDL_HTML5.Build.cs
new file mode 100644
--- /dev/null
+++ b/Engine/Platforms/HTML5/Source/Runtime/AudioMixerSDL/AudioMixerSDL_HTML5.Build.cs
@@ -0,0 +1,10 @@
+// Copyright 1998-2019 Epic Games, Inc. All Rights Reserved.
+
+using UnrealBuildTool;
+
+public class AudioMixerSDL_HTML5 : AudioMixerSDL
+{
+	public AudioMixerSDL_HTML5(ReadOnlyTargetRules Target) : base(Target)
+	{
+	}
+}
diff --git a/Engine/Platforms/HTML5/Source/Runtime/AudioMixerSDL/Private/HTML5AudioMixerSDLDefines.h b/Engine/Platforms/HTML5/Source/Runtime/AudioMixerSDL/Private/HTML5AudioMixerSDLDefines.h
new file mode 100644
--- /dev/null
+++ b/Engine/Platforms/HTML5/Source/Runtime/AudioMixerSDL/Private/HTML5AudioMixerSDLDefines.h
@@ -0,0 +1,10 @@
+// Copyright 1998-2019 Epic Games, Inc. All Rights Reserved.
+
+#pragma once
+
+#include "AudioMixerPlatformSDL.h"
+
+namespace Audio
+{
+	using FAudioMixerPlatformSDL = FMixerPlatformSDL;
+}
