add missing GetRuntimeSettingsClassName member to FHTML5PlatformProperties

#rnx

diff --git a/Engine/Platforms/HTML5/Source/Runtime/Core/Public/HTML5PlatformProperties.h b/Engine/Platforms/HTML5/Source/Runtime/Core/Public/HTML5PlatformProperties.h
--- a/Engine/Platforms/HTML5/Source/Runtime/Core/Public/HTML5PlatformProperties.h
+++ b/Engine/Platforms/HTML5/Source/Runtime/Core/Public/HTML5PlatformProperties.h
@@ -39,6 +39,11 @@ struct FHTML5PlatformProperties
 		return "HTML5";
 	}
 
+	static FORCEINLINE const TCHAR* GetRuntimeSettingsClassName()
+	{
+		return TEXT("/Script/HTML5PlatformEditor.HTML5TargetSettings");
+	}
+
 	static FORCEINLINE bool RequiresCookedData( )
 	{
 		return true;
